[Insights] Updated TraceAnalyzer to better handle text conversion for value of integer type fields, including arrays of integers.

diff --git a/Engine/Source/Programs/TraceAnalyzer/Private/ConvertToText.cpp b/Engine/Source/Programs/TraceAnalyzer/Private/ConvertToText.cpp
--- a/Engine/Source/Programs/TraceAnalyzer/Private/ConvertToText.cpp
+++ b/Engine/Source/Programs/TraceAnalyzer/Private/ConvertToText.cpp
@@ -22,9 +22,6 @@ namespace TraceAnalyzer
 
 ////////////////////////////////////////////////////////////////////////////////////////////////////
 
-constexpr int64 MinInt64Hex = -999'999;
-constexpr int64 MaxInt64Hex = 999'999'999;
-
 // See \Engine\Source\Runtime\TraceLog\Public\Trace\Detail\Protocols\Protocol6.h
 enum class EEventFlags : uint8
 {
@@ -1072,7 +1069,7 @@ bool FConvertToTextAnalyzer::OnNewEvent(uint16 RouteId, const FEventTypeInfo& Ty
 			Serializer.Append("Reference64");
 			break;
 		default:
-			Serializer.WriteValueInteger(static_cast<int64>(FieldInfo.GetType()));
+			Serializer.WriteValueUInt32(static_cast<uint32>(FieldInfo.GetType()));
 			break;
 		}
 
@@ -1279,13 +1276,66 @@ bool FConvertToTextAnalyzer::OnEvent(uint16 RouteId, EStyle Style, const FOnEven
 			{
 				case FEventFieldInfo::EType::Integer: // int64
 				{
-					if (!FieldInfo.IsSigned())
+#if 0
+					// Probe element size...
+					uint8 ElementSize = 0;
+					{
+						const TArrayReader<uint8>& Reader8 = EventData.GetArray<uint8>(FieldInfo.GetName());
+						if (Reader8.GetData())
+						{
+							ElementSize = 1;
+						}
+						else
+						{
+							const TArrayReader<uint16>& Reader16 = EventData.GetArray<uint16>(FieldInfo.GetName());
+							if (Reader16.GetData())
+							{
+								ElementSize = 2;
+							}
+							else
+							{
+								const TArrayReader<uint32>& Reader32 = EventData.GetArray<uint32>(FieldInfo.GetName());
+								if (Reader16.GetData())
+								{
+									ElementSize = 4;
+								}
+								else
+								{
+									ElementSize = 8;
+								}
+							}
+						}
+					}
+#endif
+
+					if (FieldInfo.IsSigned())
+					{
+						Serializer.BeginArray();
+						const TArrayReader<int64>& Reader = EventData.GetArray<int64>(FieldInfo.GetName());
+						const uint32 ArrayCount = Reader.Num();
+						const uint32 ActualCount = FMath::Min(ArrayCount, 8u);
+						for (uint32 ArrayIndex = 0; ArrayIndex < ActualCount; ++ArrayIndex)
+						{
+							if (ArrayIndex != 0)
+							{
+								Serializer.NextArrayElement();
+							}
+							const int64 Value = Reader[ArrayIndex];
+							Serializer.WriteValueInt64Auto(Value);
+						}
+						if (ActualCount != ArrayCount)
+						{
+							Serializer.Appendf(" ... | %u elements", ArrayCount);
+						}
+						Serializer.EndArray();
+					}
+					else // unsigned
 					{
-						// Write array as 0x[00 01 ... FF]
 						const TArrayReader<uint8>& Reader8 = EventData.GetArray<uint8>(FieldInfo.GetName());
 						const uint8* Data8 = Reader8.GetData();
 						if (Data8)
 						{
+							// Write array as 0x[00 01 ... FF]
 							Serializer.Append("0x");
 							Serializer.BeginArray();
 							const uint32 ArrayCount = Reader8.Num();
@@ -1308,33 +1358,26 @@ bool FConvertToTextAnalyzer::OnEvent(uint16 RouteId, EStyle Style, const FOnEven
 							Serializer.EndArray();
 							break;
 						}
-					}
 
-					Serializer.BeginArray();
-					const TArrayReader<int64>& Reader = EventData.GetArray<int64>(FieldInfo.GetName());
-					const uint32 ArrayCount = Reader.Num();
-					const uint32 ActualCount = FMath::Min(ArrayCount, 4u);
-					for (uint32 ArrayIndex = 0; ArrayIndex < ActualCount; ++ArrayIndex)
-					{
-						if (ArrayIndex != 0)
-						{
-							Serializer.NextArrayElement();
-						}
-						const int64 Value = Reader[ArrayIndex];
-						if (Value < MinInt64Hex || Value > MaxInt64Hex)
+						Serializer.BeginArray();
+						const TArrayReader<uint64>& Reader = EventData.GetArray<uint64>(FieldInfo.GetName());
+						const uint32 ArrayCount = Reader.Num();
+						const uint32 ActualCount = FMath::Min(ArrayCount, 8u);
+						for (uint32 ArrayIndex = 0; ArrayIndex < ActualCount; ++ArrayIndex)
 						{
-							Serializer.WriteValueIntegerHex(Value);
+							if (ArrayIndex != 0)
+							{
+								Serializer.NextArrayElement();
+							}
+							const uint64 Value = Reader[ArrayIndex];
+							Serializer.WriteValueUInt64Auto(Value);
 						}
-						else
+						if (ActualCount != ArrayCount)
 						{
-							Serializer.WriteValueInteger(Value);
+							Serializer.Appendf(" ... | %u elements", ArrayCount);
 						}
+						Serializer.EndArray();
 					}
-					if (ActualCount != ArrayCount)
-					{
-						Serializer.Appendf(" ... | %u elements", ArrayCount);
-					}
-					Serializer.EndArray();
 					break;
 				}
 
@@ -1390,22 +1433,37 @@ bool FConvertToTextAnalyzer::OnEvent(uint16 RouteId, EStyle Style, const FOnEven
 			{
 				case FEventFieldInfo::EType::Integer:
 				{
-					const int64 Value = EventData.GetValue<int64>(FieldInfo.GetName());
-					if (strcmp(FieldInfo.GetName(), "Cycle") == 0)
-					{
-						Serializer.WriteValueInteger(uint64(Value));
-						Serializer.AppendChar('(');
-						double Time = Context.EventTime.AsSeconds(uint64(Value));
-						Serializer.WriteValueTime(Time);
-						Serializer.AppendChar(')');
-					}
-					else if (Value < MinInt64Hex || Value > MaxInt64Hex)
+					if (FieldInfo.IsSigned())
 					{
-						Serializer.WriteValueIntegerHex(Value);
+						const int64 Value = EventData.GetValue<int64>(FieldInfo.GetName());
+						if (strcmp(FieldInfo.GetName(), "Cycle") == 0)
+						{
+							Serializer.WriteValueUInt64(uint64(Value));
+							Serializer.AppendChar('(');
+							double Time = Context.EventTime.AsSeconds(uint64(Value));
+							Serializer.WriteValueTime(Time);
+							Serializer.AppendChar(')');
+						}
+						else
+						{
+							Serializer.WriteValueInt64Auto(Value);
+						}
 					}
-					else
+					else // unsigned
 					{
-						Serializer.WriteValueInteger(Value);
+						const uint64 Value = EventData.GetValue<uint64>(FieldInfo.GetName());
+						if (strcmp(FieldInfo.GetName(), "Cycle") == 0)
+						{
+							Serializer.WriteValueUInt64(Value);
+							Serializer.AppendChar('(');
+							double Time = Context.EventTime.AsSeconds(Value);
+							Serializer.WriteValueTime(Time);
+							Serializer.AppendChar(')');
+						}
+						else
+						{
+							Serializer.WriteValueUInt64Auto(Value);
+						}
 					}
 					break;
 				}
diff --git a/Engine/Source/Programs/TraceAnalyzer/Private/TextSerializer.cpp b/Engine/Source/Programs/TraceAnalyzer/Private/TextSerializer.cpp
--- a/Engine/Source/Programs/TraceAnalyzer/Private/TextSerializer.cpp
+++ b/Engine/Source/Programs/TraceAnalyzer/Private/TextSerializer.cpp
@@ -3,6 +3,7 @@
 #include "TextSerializer.h"
 
 #include "HAL/UnrealMemory.h"
+#include "Math/NumericLimits.h"
 #include "Misc/CString.h"
 
 namespace UE
@@ -26,6 +27,65 @@ FTextSerializer::~FTextSerializer()
 
 ////////////////////////////////////////////////////////////////////////////////////////////////////
 
+void FTextSerializer::WriteValueInt64Auto(int64 Value)
+{
+	constexpr int64 MinInt64Hex = -99'999'999;
+	constexpr int64 MaxInt64Hex = 999'999'999;
+
+	if (Value < MinInt64Hex)
+	{
+		if (Value == MIN_int64)
+		{
+			Append("MIN_int64");
+		}
+		else
+		{
+			AppendChar('-');
+			WriteValueHex64(uint64(-Value));
+		}
+	}
+	else if (Value > MaxInt64Hex)
+	{
+		if (Value == MAX_int64)
+		{
+			Append("MAX_int64");
+		}
+		else
+		{
+			WriteValueHex64(uint64(Value));
+		}
+	}
+	else
+	{
+		WriteValueInt64(Value);
+	}
+}
+
+////////////////////////////////////////////////////////////////////////////////////////////////////
+
+void FTextSerializer::WriteValueUInt64Auto(uint64 Value)
+{
+	constexpr uint64 MaxUInt64Hex = 999'999'999;
+
+	if (Value > MaxUInt64Hex)
+	{
+		if (Value == MAX_uint64)
+		{
+			Append("MAX_uint64");
+		}
+		else
+		{
+			WriteValueHex64(Value);
+		}
+	}
+	else
+	{
+		WriteValueUInt64(Value);
+	}
+}
+
+////////////////////////////////////////////////////////////////////////////////////////////////////
+
 void FTextSerializer::WriteValueBinary(const void* Data, uint32 Size)
 {
 	AppendChar('<');
diff --git a/Engine/Source/Programs/TraceAnalyzer/Private/TextSerializer.h b/Engine/Source/Programs/TraceAnalyzer/Private/TextSerializer.h
--- a/Engine/Source/Programs/TraceAnalyzer/Private/TextSerializer.h
+++ b/Engine/Source/Programs/TraceAnalyzer/Private/TextSerializer.h
@@ -168,14 +168,28 @@ public:
 	void WriteValueReference(const UE::Trace::FEventRef64& Value) { Appendf("R(%u,%llu)", Value.RefTypeId, Value.Id); }
 
 	void WriteValueBool(bool Value)         { Append(Value ? "true" : "false"); }
-	void WriteValueInteger(int64 Value)     { Appendf("%lli", Value); }
-	void WriteValueIntegerHex(int64 Value)  { Appendf("0x%llX", uint64(Value)); }
+
+	void WriteValueInt8(int8 Value)         { Appendf("%i", int32(Value)); }
+	void WriteValueInt16(int16 Value)       { Appendf("%i", int32(Value)); }
+	void WriteValueInt32(int32 Value)       { Appendf("%i", Value); }
+	void WriteValueInt64(int64 Value)       { Appendf("%lli", Value); }
+
+	void WriteValueUInt8(uint8 Value)       { Appendf("%u", uint32(Value)); }
+	void WriteValueUInt16(uint16 Value)     { Appendf("%u", uint32(Value)); }
+	void WriteValueUInt32(uint32 Value)     { Appendf("%u", Value); }
+	void WriteValueUInt64(uint64 Value)     { Appendf("%llu", Value); }
+
 	void WriteValueHex8(uint8 Value)        { Appendf("0x%X", uint32(Value)); }
 	void WriteValueHex16(uint16 Value)      { Appendf("0x%X", uint32(Value)); }
 	void WriteValueHex32(uint32 Value)      { Appendf("0x%X", Value); }
 	void WriteValueHex64(uint64 Value)      { Appendf("0x%llX", Value); }
+
+	void WriteValueInt64Auto(int64 Value);
+	void WriteValueUInt64Auto(uint64 Value);
+
 	void WriteValueFloat(float Value)       { Appendf("%f", Value); }
 	void WriteValueDouble(double Value)     { Appendf("%f", Value); }
+
 	void WriteValueTime(double Time)        { Appendf("%f", Time); }
 	void WriteValueNull()                   { Append("null"); }
 
@@ -193,7 +207,7 @@ public:
 	void WriteString(const ANSICHAR* Name, const ANSICHAR* Value)             { WriteKey(Name); WriteValueString(Value); }
 	void WriteString(const ANSICHAR* Name, const ANSICHAR* Value, uint32 Len) { WriteKey(Name); WriteValueString(Value, Len); }
 	void WriteBool(const ANSICHAR* Name, bool Value)                          { WriteKey(Name); WriteValueBool(Value); }
-	void WriteInteger(const ANSICHAR* Name, int64 Value)                      { WriteKey(Name); WriteValueInteger(Value); }
+	void WriteInteger(const ANSICHAR* Name, int64 Value)                      { WriteKey(Name); WriteValueInt64(Value); }
 	void WriteIntegerHex(const ANSICHAR* Name, int64 Value)                   { WriteKey(Name); WriteValueHex64(uint64(Value)); }
 	void WriteFloat(const ANSICHAR* Name, float Value)                        { WriteKey(Name); WriteValueFloat(Value); }
 	void WriteDouble(const ANSICHAR* Name, double Value)                      { WriteKey(Name); WriteValueDouble(Value); }
@@ -212,7 +226,7 @@ class FStdoutTextSerializer : public FTextSerializer
 public:
 	FStdoutTextSerializer();
 	virtual ~FStdoutTextSerializer() {}
-	
+
 	virtual void AppendChar(const ANSICHAR Value) override;
 	virtual void Append(const ANSICHAR* Text, int32 Len) override;
 	virtual void Append(const ANSICHAR* Text) override;
