[Insights] TraceAnalyzer: Fixed/improved text serialization of trace events (ex.: when an event does not have fields).

diff --git a/Engine/Source/Programs/TraceAnalyzer/Private/ConvertToText.cpp b/Engine/Source/Programs/TraceAnalyzer/Private/ConvertToText.cpp
--- a/Engine/Source/Programs/TraceAnalyzer/Private/ConvertToText.cpp
+++ b/Engine/Source/Programs/TraceAnalyzer/Private/ConvertToText.cpp
@@ -447,7 +447,7 @@ void FConvertToTextAnalyzer::RegisterCallbacksForKnownEvents()
 			auto NameTChar = StringCast<TCHAR>(Name.GetData(), Name.Num());
 			FStringView NameView(NameTChar.Get(), NameTChar.Length());
 			FString NameStr(NameView.GetData(), NameView.Len());
-			Serializer.WriteString("#Name", TCHAR_TO_UTF8(*NameStr));
+			Serializer.WriteAttributeString("#Name", TCHAR_TO_UTF8(*NameStr));
 		};
 
 	KnownEvents[(int)EKnownEvent::CpuProfiler_EventSpec].Callback =
@@ -477,11 +477,11 @@ void FConvertToTextAnalyzer::RegisterCallbacksForKnownEvents()
 				}
 				if (TimerName)
 				{
-					Serializer.WriteString("@Name", TCHAR_TO_UTF8(TimerName));
+					Serializer.WriteAttributeString("@Name", TCHAR_TO_UTF8(TimerName));
 				}
 				else
 				{
-					Serializer.WriteString("@Name", "<invalid>");
+					Serializer.WriteAttributeString("@Name", "<invalid>");
 				}
 			}
 
@@ -496,14 +496,18 @@ void FConvertToTextAnalyzer::RegisterCallbacksForKnownEvents()
 			TArrayView<const uint8> Out = EventData.GetArrayView<uint8>("Data");
 			if (Out.GetData() != nullptr)
 			{
+				Serializer.BeginAttribute();
 				Serializer.WriteKey("#Data");
 				ConvertToTextEventBatch(Context, Serializer, Out.GetData(), Out.Num());
+				Serializer.EndAttribute();
 			}
 			else
 			{
 				check(EventData.GetAttachment() != nullptr);
+				Serializer.BeginAttribute();
 				Serializer.WriteKey("Attached");
 				ConvertToTextEventBatch(Context, Serializer, EventData.GetAttachment(), EventData.GetAttachmentSize());
+				Serializer.EndAttribute();
 			}
 		};
 
@@ -514,8 +518,10 @@ void FConvertToTextAnalyzer::RegisterCallbacksForKnownEvents()
 			TArrayView<const uint8> Out = Context.EventData.GetArrayView<uint8>("Data");
 			if (Out.GetData() != nullptr)
 			{
+				Serializer.BeginAttribute();
 				Serializer.WriteKey("#Data");
 				ConvertToTextEventBatchV2(Context, Serializer, Out.GetData(), Out.Num());
+				Serializer.EndAttribute();
 			}
 		};
 
@@ -538,8 +544,7 @@ void FConvertToTextAnalyzer::RegisterCallbacksForKnownEvents()
 	//	{
 	//		uint64 Cycle = Context.EventData.GetValue<uint64>("Cycle");
 	//		double Time = Context.EventTime.AsSeconds(Cycle);
-	//		Serializer.NextEventField();
-	//		Serializer.WriteFloat("#Time", Time);
+	//		Serializer.WriteAttributeFloat("#Time", Time);
 	//	};
 }
 
@@ -556,7 +561,7 @@ void FConvertToTextAnalyzer::RegisterAttachmentCallbacksForKnownEvents()
 			if (TypeInfo.GetFieldIndex("Name") < 0)
 			{
 				const ANSICHAR* ChannelName = reinterpret_cast<const ANSICHAR*>(EventData.GetAttachment());
-				Serializer.WriteString("@ChannelName", ChannelName);
+				Serializer.WriteAttributeString("@ChannelName", ChannelName);
 			}
 		};
 
@@ -568,7 +573,7 @@ void FConvertToTextAnalyzer::RegisterAttachmentCallbacksForKnownEvents()
 			if (TypeInfo.GetFieldIndex("Name") < 0)
 			{
 				const ANSICHAR* ThreadName = reinterpret_cast<const ANSICHAR*>(EventData.GetAttachment());
-				Serializer.WriteString("@Name", ThreadName);
+				Serializer.WriteAttributeString("@Name", ThreadName);
 			}
 		};
 
@@ -580,17 +585,13 @@ void FConvertToTextAnalyzer::RegisterAttachmentCallbacksForKnownEvents()
 			const uint8 CommandLineOffset = EventData.GetValue<uint8>("CommandLineOffset");
 
 			const ANSICHAR* Platform = reinterpret_cast<const ANSICHAR*>(EventData.GetAttachment());
-			Serializer.WriteString("@Platform", Platform, AppNameOffset);
-
-			Serializer.NextEventField();
+			Serializer.WriteAttributeString("@Platform", Platform, AppNameOffset);
 
 			const ANSICHAR* AppName = reinterpret_cast<const ANSICHAR*>(EventData.GetAttachment() + AppNameOffset);
-			Serializer.WriteString("@AppName", AppName, CommandLineOffset - AppNameOffset);
-
-			Serializer.NextEventField();
+			Serializer.WriteAttributeString("@AppName", AppName, CommandLineOffset - AppNameOffset);
 
 			const ANSICHAR* CommandLine = reinterpret_cast<const ANSICHAR*>(EventData.GetAttachment() + CommandLineOffset);
-			Serializer.WriteString("@CommandLine", CommandLine, EventData.GetAttachmentSize() - CommandLineOffset);
+			Serializer.WriteAttributeString("@CommandLine", CommandLine, EventData.GetAttachmentSize() - CommandLineOffset);
 		};
 
 	KnownEvents[(int)EKnownEvent::Diagnostics_Session2].AttachmentCallback =
@@ -604,7 +605,7 @@ void FConvertToTextAnalyzer::RegisterAttachmentCallbacksForKnownEvents()
 		{
 			const FEventData& EventData = Context.EventData;
 			const TCHAR* ThreadName = reinterpret_cast<const TCHAR*>(EventData.GetAttachment());
-			Serializer.WriteString("@Name", TCHAR_TO_UTF8(ThreadName));
+			Serializer.WriteAttributeString("@Name", TCHAR_TO_UTF8(ThreadName));
 		};
 
 	KnownEvents[(int)EKnownEvent::Misc_BeginGameFrame].AttachmentCallback =
@@ -616,7 +617,7 @@ void FConvertToTextAnalyzer::RegisterAttachmentCallbacksForKnownEvents()
 			const FEventData& EventData = Context.EventData;
 			const uint8* BufferPtr = EventData.GetAttachment();
 			uint64 CycleDiff = Decode7bit(BufferPtr);
-			Serializer.WriteInteger("@CycleDiff", CycleDiff);
+			Serializer.WriteAttributeInteger("@CycleDiff", CycleDiff);
 		};
 
 	KnownEvents[(int)EKnownEvent::Trace_ThreadGroupBegin].AttachmentCallback =
@@ -629,7 +630,7 @@ void FConvertToTextAnalyzer::RegisterAttachmentCallbacksForKnownEvents()
 			if (TypeInfo.GetFieldIndex("Name") < 0)
 			{
 				const ANSICHAR* GroupName = reinterpret_cast<const ANSICHAR*>(EventData.GetAttachment());
-				Serializer.WriteString("@GroupName", GroupName);
+				Serializer.WriteAttributeString("@GroupName", GroupName);
 			}
 		};
 
@@ -642,12 +643,10 @@ void FConvertToTextAnalyzer::RegisterAttachmentCallbacksForKnownEvents()
 			if (TypeInfo.GetFieldIndex("FileName") < 0)
 			{
 				const ANSICHAR* File = reinterpret_cast<const ANSICHAR*>(EventData.GetAttachment());
-				Serializer.WriteString("@File", File);
-
-				Serializer.NextEventField();
+				Serializer.WriteAttributeString("@File", File);
 
 				const TCHAR* FormatString = reinterpret_cast<const TCHAR*>(EventData.GetAttachment() + strlen(File) + 1);
-				Serializer.WriteString("@FormatString", TCHAR_TO_UTF8(FormatString));
+				Serializer.WriteAttributeString("@FormatString", TCHAR_TO_UTF8(FormatString));
 			}
 		};
 
@@ -659,7 +658,7 @@ void FConvertToTextAnalyzer::RegisterAttachmentCallbacksForKnownEvents()
 			if (TypeInfo.GetFieldIndex("Name") < 0)
 			{
 				const TCHAR* CategoryName = reinterpret_cast<const TCHAR*>(EventData.GetAttachment());
-				Serializer.WriteString("@CategoryName", TCHAR_TO_UTF8(CategoryName));
+				Serializer.WriteAttributeString("@CategoryName", TCHAR_TO_UTF8(CategoryName));
 			}
 		};
 
@@ -668,7 +667,7 @@ void FConvertToTextAnalyzer::RegisterAttachmentCallbacksForKnownEvents()
 		{
 			const FEventData& EventData = Context.EventData;
 			const TCHAR* ClassName = reinterpret_cast<const TCHAR*>(EventData.GetAttachment());
-			Serializer.WriteString("@File", TCHAR_TO_UTF8(ClassName));
+			Serializer.WriteAttributeString("@File", TCHAR_TO_UTF8(ClassName));
 		};
 
 	KnownEvents[(int)EKnownEvent::LoadTime_ClassInfo].AttachmentCallback =
@@ -679,7 +678,7 @@ void FConvertToTextAnalyzer::RegisterAttachmentCallbacksForKnownEvents()
 			if (TypeInfo.GetFieldIndex("Name") < 0)
 			{
 				const TCHAR* ClassName = reinterpret_cast<const TCHAR*>(EventData.GetAttachment());
-				Serializer.WriteString("@ClassName", TCHAR_TO_UTF8(ClassName));
+				Serializer.WriteAttributeString("@ClassName", TCHAR_TO_UTF8(ClassName));
 			}
 		};
 
@@ -688,7 +687,7 @@ void FConvertToTextAnalyzer::RegisterAttachmentCallbacksForKnownEvents()
 		{
 			const FEventData& EventData = Context.EventData;
 			const TCHAR* ClassName = reinterpret_cast<const TCHAR*>(EventData.GetAttachment());
-			Serializer.WriteString("@PackageName", TCHAR_TO_UTF8(ClassName));
+			Serializer.WriteAttributeString("@PackageName", TCHAR_TO_UTF8(ClassName));
 		};
 
 	KnownEvents[(int)EKnownEvent::CpuProfiler_EventSpec].AttachmentCallback =
@@ -718,11 +717,11 @@ void FConvertToTextAnalyzer::RegisterAttachmentCallbacksForKnownEvents()
 				}
 				if (TimerName)
 				{
-					Serializer.WriteString("@Name", TCHAR_TO_UTF8(TimerName));
+					Serializer.WriteAttributeString("@Name", TCHAR_TO_UTF8(TimerName));
 				}
 				else
 				{
-					Serializer.WriteString("@Name", "<invalid>");
+					Serializer.WriteAttributeString("@Name", "<invalid>");
 				}
 			}
 
@@ -735,8 +734,10 @@ void FConvertToTextAnalyzer::RegisterAttachmentCallbacksForKnownEvents()
 		{
 			const FEventData& EventData = Context.EventData;
 			check(EventData.GetAttachment() != nullptr);
+			Serializer.BeginAttribute();
 			Serializer.WriteKey("Attached");
 			ConvertToTextEventBatch(Context, Serializer, EventData.GetAttachment(), EventData.GetAttachmentSize());
+			Serializer.EndAttribute();
 		};
 
 	KnownEvents[(int)EKnownEvent::Stats_Spec].AttachmentCallback =
@@ -747,12 +748,10 @@ void FConvertToTextAnalyzer::RegisterAttachmentCallbacksForKnownEvents()
 			if (TypeInfo.GetFieldIndex("Name") < 0)
 			{
 				const ANSICHAR* Name = reinterpret_cast<const ANSICHAR*>(EventData.GetAttachment());
-				Serializer.WriteString("@Name", Name);
-
-				Serializer.NextEventField();
+				Serializer.WriteAttributeString("@Name", Name);
 
 				const TCHAR* Description = reinterpret_cast<const TCHAR*>(EventData.GetAttachment() + strlen(Name) + 1);
-				Serializer.WriteString("@Description", TCHAR_TO_UTF8(Description));
+				Serializer.WriteAttributeString("@Description", TCHAR_TO_UTF8(Description));
 			}
 		};
 
@@ -764,7 +763,7 @@ void FConvertToTextAnalyzer::RegisterAttachmentCallbacksForKnownEvents()
 			if (TypeInfo.GetFieldIndex("Name") < 0)
 			{
 				const TCHAR* Name = reinterpret_cast<const TCHAR*>(EventData.GetAttachment());
-				Serializer.WriteString("@Name", TCHAR_TO_UTF8(Name));
+				Serializer.WriteAttributeString("@Name", TCHAR_TO_UTF8(Name));
 			}
 		};
 
@@ -776,7 +775,7 @@ void FConvertToTextAnalyzer::RegisterAttachmentCallbacksForKnownEvents()
 			if (TypeInfo.GetFieldIndex("Name") < 0)
 			{
 				const TCHAR* Name = reinterpret_cast<const TCHAR*>(EventData.GetAttachment());
-				Serializer.WriteString("@Name", TCHAR_TO_UTF8(Name));
+				Serializer.WriteAttributeString("@Name", TCHAR_TO_UTF8(Name));
 			}
 		};
 
@@ -788,7 +787,7 @@ void FConvertToTextAnalyzer::RegisterAttachmentCallbacksForKnownEvents()
 			if (TypeInfo.GetFieldIndex("Name") < 0)
 			{
 				const ANSICHAR* Name = reinterpret_cast<const ANSICHAR*>(EventData.GetAttachment());
-				Serializer.WriteString("@Name", Name);
+				Serializer.WriteAttributeString("@Name", Name);
 			}
 		};
 
@@ -800,7 +799,7 @@ void FConvertToTextAnalyzer::RegisterAttachmentCallbacksForKnownEvents()
 			if (TypeInfo.GetFieldIndex("Name") < 0)
 			{
 				const TCHAR* Name = reinterpret_cast<const TCHAR*>(EventData.GetAttachment());
-				Serializer.WriteString("@Name", TCHAR_TO_UTF8(Name));
+				Serializer.WriteAttributeString("@Name", TCHAR_TO_UTF8(Name));
 			}
 		};
 
@@ -810,9 +809,8 @@ void FConvertToTextAnalyzer::RegisterAttachmentCallbacksForKnownEvents()
 			const FEventData& EventData = Context.EventData;
 			const TCHAR* Key = reinterpret_cast<const TCHAR*>(EventData.GetAttachment());
 			const TCHAR* Value = reinterpret_cast<const TCHAR*>(EventData.GetAttachment() + EventData.GetValue<uint16>("ValueOffset"));
-			Serializer.WriteString("@Key", TCHAR_TO_UTF8(Key));
-			Serializer.NextEventField();
-			Serializer.WriteString("@Value", TCHAR_TO_UTF8(Value));
+			Serializer.WriteAttributeString("@Key", TCHAR_TO_UTF8(Key));
+			Serializer.WriteAttributeString("@Value", TCHAR_TO_UTF8(Value));
 		};
 
 	KnownEvents[(int)EKnownEvent::CsvProfiler_BeginCapture].AttachmentCallback =
@@ -820,7 +818,7 @@ void FConvertToTextAnalyzer::RegisterAttachmentCallbacksForKnownEvents()
 		{
 			const FEventData& EventData = Context.EventData;
 			const TCHAR* Filename = reinterpret_cast<const TCHAR*>(EventData.GetAttachment());
-			Serializer.WriteString("@Filename", TCHAR_TO_UTF8(Filename));
+			Serializer.WriteAttributeString("@Filename", TCHAR_TO_UTF8(Filename));
 		};
 }
 
@@ -984,18 +982,15 @@ bool FConvertToTextAnalyzer::OnNewEvent(uint16 RouteId, const FEventTypeInfo& Ty
 	{
 		Serializer.BeginNewEventHeader();
 
-		Serializer.WriteInteger("Id", TypeInfo.GetId());
+		Serializer.WriteAttributeInteger("Id", TypeInfo.GetId());
 
-		Serializer.NextEventField();
-		Serializer.WriteString("LoggerName", TypeInfo.GetLoggerName());
+		Serializer.WriteAttributeString("LoggerName", TypeInfo.GetLoggerName());
 
-		Serializer.NextEventField();
-		Serializer.WriteString("Name", TypeInfo.GetName());
+		Serializer.WriteAttributeString("Name", TypeInfo.GetName());
 
-		Serializer.NextEventField();
+		Serializer.BeginAttribute();
 		Serializer.WriteKey("Flags");
 		Serializer.Append("\"");
-
 		const uint8 EventFlags = TypeInfo.GetFlags();
 		if ((EventFlags & (uint8)EEventFlags::Definition) != 0)
 		{
@@ -1018,6 +1013,7 @@ bool FConvertToTextAnalyzer::OnNewEvent(uint16 RouteId, const FEventTypeInfo& Ty
 			Serializer.Append("Sync");
 		}
 		Serializer.Append("\"");
+		Serializer.EndAttribute();
 
 		Serializer.EndNewEventHeader();
 	}
@@ -1029,17 +1025,15 @@ bool FConvertToTextAnalyzer::OnNewEvent(uint16 RouteId, const FEventTypeInfo& Ty
 		Serializer.BeginField();
 
 		const FEventFieldInfo& FieldInfo = *(TypeInfo.GetFieldInfo(FieldIndex));
-		Serializer.WriteString("Name", FieldInfo.GetName());
+		Serializer.WriteAttributeString("Name", FieldInfo.GetName());
 
-		Serializer.NextEventField();
 		/** Offset from the start of the event to this field's data. */
-		Serializer.WriteInteger("Offset", FieldInfo.GetOffset());
+		Serializer.WriteAttributeInteger("Offset", FieldInfo.GetOffset());
 
-		Serializer.NextEventField();
 		/** The size of the field's data in bytes. */
-		Serializer.WriteInteger("Size", FieldInfo.GetSize());
+		Serializer.WriteAttributeInteger("Size", FieldInfo.GetSize());
 
-		Serializer.NextEventField();
+		Serializer.BeginAttribute();
 		/** What type of field is this? */
 		Serializer.WriteKey("Type");
 		switch (FieldInfo.GetType())
@@ -1072,10 +1066,10 @@ bool FConvertToTextAnalyzer::OnNewEvent(uint16 RouteId, const FEventTypeInfo& Ty
 			Serializer.WriteValueUInt32(static_cast<uint32>(FieldInfo.GetType()));
 			break;
 		}
+		Serializer.EndAttribute();
 
-		Serializer.NextEventField();
 		/** Is this field an array-type field? */
-		Serializer.WriteBool("IsArray", FieldInfo.IsArray());
+		Serializer.WriteAttributeBool("IsArray", FieldInfo.IsArray());
 
 		Serializer.EndField();
 	}
@@ -1206,11 +1200,7 @@ bool FConvertToTextAnalyzer::OnEvent(uint16 RouteId, EStyle Style, const FOnEven
 
 	if (RouteId == 1)
 	{
-		if (Serializer.IsWriteEventHeaderEnabled())
-		{
-			Serializer.AppendChar(' ');
-		}
-
+		Serializer.BeginAttribute();
 		if (Style == EStyle::EnterScope)
 		{
 			Serializer.Append("!EnterScope");
@@ -1229,33 +1219,17 @@ bool FConvertToTextAnalyzer::OnEvent(uint16 RouteId, EStyle Style, const FOnEven
 		{
 			Serializer.Appendf("!%u", (uint32)Style);
 		}
-
-		if (!Serializer.IsWriteEventHeaderEnabled() && (FieldCount > 0 || bHasAttachment))
-		{
-			Serializer.AppendChar(' ');
-		}
+		Serializer.EndAttribute();
 	}
 
 	Serializer.WriteEventName(TypeInfo.GetLoggerName(), TypeInfo.GetName());
 
 	//////////////////////////////////////////////////
 
-	if (FieldCount > 0 || bHasAttachment)
-	{
-		Serializer.BeginEventFields();
-	}
-
-	//////////////////////////////////////////////////
-
 	for (uint32 FieldIndex = 0; FieldIndex < FieldCount; ++FieldIndex)
 	{
 		const FEventFieldInfo& FieldInfo = *(TypeInfo.GetFieldInfo(FieldIndex));
 
-		if (FieldIndex != 0)
-		{
-			Serializer.NextEventField();
-		}
-
 		if (FieldInfo.IsArray())
 		{
 			if (IsKnownEvent(EventId, EKnownEvent::GpuProfiler_EventSpec) || // "Name" array
@@ -1268,8 +1242,8 @@ bool FConvertToTextAnalyzer::OnEvent(uint16 RouteId, EStyle Style, const FOnEven
 			}
 		}
 
+		Serializer.BeginAttribute();
 		Serializer.WriteKey(FieldInfo.GetName());
-
 		if (FieldInfo.IsArray())
 		{
 			switch (FieldInfo.GetType())
@@ -1514,8 +1488,7 @@ bool FConvertToTextAnalyzer::OnEvent(uint16 RouteId, EStyle Style, const FOnEven
 				}
 			}
 		}
-
-		//Serializer.NextArrayElement();
+		Serializer.EndAttribute();
 	}
 
 	//////////////////////////////////////////////////
@@ -1532,21 +1505,23 @@ bool FConvertToTextAnalyzer::OnEvent(uint16 RouteId, EStyle Style, const FOnEven
 
 	if (bHasAttachment)
 	{
-		Serializer.NextEventField();
-
 		if (Event.AttachmentCallback)
 		{
 			Event.AttachmentCallback(Context, Serializer);
 		}
 		else
 		{
-			Serializer.WriteBinary("Attached", EventData.GetAttachment(), AttachmentSize);
+			Serializer.WriteAttributeBinary("Attached", EventData.GetAttachment(), AttachmentSize);
 		}
 	}
 
 	if (bEmitSizeStats)
 	{
-		Serializer.Appendf("\n\t\t> SIZE Fix=%u Aux=%u Att=%u Raw=%u Total=%u", FixedSize, AuxSize, AttachmentSize, RawSize, TotalEventSize);
+		Serializer.AppendChar('\n');
+		Serializer.BeginAttributeSet();
+		Serializer.BeginAttribute();
+		Serializer.Appendf("> SIZE Fix=%u Aux=%u Att=%u Raw=%u Total=%u", FixedSize, AuxSize, AttachmentSize, RawSize, TotalEventSize);
+		Serializer.EndAttribute();
 	}
 
 	//////////////////////////////////////////////////
@@ -1581,7 +1556,7 @@ void FConvertToTextAnalyzer::ConvertToTextEventBatch(const FOnEventContext& Cont
 	{
 		if (!bFirstElement)
 		{
-			Serializer.NextEventField();
+			Serializer.AppendChar(' ');
 		}
 		else
 		{
@@ -1620,7 +1595,8 @@ void FConvertToTextAnalyzer::ConvertToTextEventBatch(const FOnEventContext& Cont
 			//Serializer.WriteInteger("B", ActualCycle);
 			if (bIsAbsoluteTimestamp)
 			{
-				Serializer.WriteInteger("*B", Cycle);
+				Serializer.Append("*B=");
+				Serializer.WriteValueInt64(Cycle);
 				Serializer.AppendChar('(');
 				double Time = Context.EventTime.AsSeconds(Cycle);
 				Serializer.WriteValueTime(Time);
@@ -1628,11 +1604,12 @@ void FConvertToTextAnalyzer::ConvertToTextEventBatch(const FOnEventContext& Cont
 			}
 			else
 			{
-				Serializer.WriteInteger("B", Cycle);
+				Serializer.Append("B=");
+				Serializer.WriteValueInt64(Cycle);
 			}
 
-			Serializer.NextEventField();
-			Serializer.WriteInteger(bIsUnknownSpecId ? "?id" : "id", SpecId);
+			Serializer.Append(bIsUnknownSpecId ? " ?id=" : " id=");
+			Serializer.WriteValueInt64(SpecId);
 		}
 		else
 		{
@@ -1641,7 +1618,8 @@ void FConvertToTextAnalyzer::ConvertToTextEventBatch(const FOnEventContext& Cont
 			//Serializer.WriteInteger("E", ActualCycle);
 			if (bIsAbsoluteTimestamp)
 			{
-				Serializer.WriteInteger("*E", Cycle);
+				Serializer.Append("*E=");
+				Serializer.WriteValueInt64(Cycle);
 				Serializer.AppendChar('(');
 				double Time = Context.EventTime.AsSeconds(Cycle);
 				Serializer.WriteValueTime(Time);
@@ -1649,7 +1627,8 @@ void FConvertToTextAnalyzer::ConvertToTextEventBatch(const FOnEventContext& Cont
 			}
 			else
 			{
-				Serializer.WriteInteger("E", Cycle);
+				Serializer.Append("E=");
+				Serializer.WriteValueInt64(Cycle);
 			}
 		}
 	}
@@ -1688,7 +1667,7 @@ void FConvertToTextAnalyzer::ConvertToTextEventBatchV2(const FOnEventContext& Co
 	{
 		if (!bFirstElement)
 		{
-			Serializer.NextEventField();
+			Serializer.AppendChar(' ');
 		}
 		else
 		{
@@ -1705,20 +1684,22 @@ void FConvertToTextAnalyzer::ConvertToTextEventBatchV2(const FOnEventContext& Co
 		{
 			if (DecodedCycle & 1ull) // Restore
 			{
-				Serializer.WriteInteger("coR", Cycle);
+				Serializer.Append("coR=");
+				Serializer.WriteValueInt64(Cycle);
 				uint64 CoroutineId = Decode7bit(BufferPtr);
-				Serializer.NextEventField();
-				Serializer.WriteInteger("coId", CoroutineId);
+				Serializer.Append(" coId=");
+				Serializer.WriteValueInt64(CoroutineId);
 				uint32 TimerScopeDepth = Decode7bit(BufferPtr);
-				Serializer.NextEventField();
-				Serializer.WriteInteger("coDepth", TimerScopeDepth);
+				Serializer.Append(" coDepth=");
+				Serializer.WriteValueInt64(TimerScopeDepth);
 			}
 			else // Save
 			{
-				Serializer.WriteInteger("coS", Cycle);
+				Serializer.Append("coS=");
+				Serializer.WriteValueInt64(Cycle);
 				uint32 TimerScopeDepth = Decode7bit(BufferPtr);
-				Serializer.NextEventField();
-				Serializer.WriteInteger("coDepth", TimerScopeDepth);
+				Serializer.Append(" coDepth=");
+				Serializer.WriteValueInt64(TimerScopeDepth);
 			}
 		}
 		else // normal CPU events
@@ -1743,7 +1724,8 @@ void FConvertToTextAnalyzer::ConvertToTextEventBatchV2(const FOnEventContext& Co
 				//Serializer.WriteInteger("B", ActualCycle);
 				if (bIsAbsoluteTimestamp)
 				{
-					Serializer.WriteInteger("*B", Cycle);
+					Serializer.Append("*B=");
+					Serializer.WriteValueInt64(Cycle);
 					Serializer.AppendChar('(');
 					double Time = Context.EventTime.AsSeconds(Cycle);
 					Serializer.WriteValueTime(Time);
@@ -1751,11 +1733,12 @@ void FConvertToTextAnalyzer::ConvertToTextEventBatchV2(const FOnEventContext& Co
 				}
 				else
 				{
-					Serializer.WriteInteger("B", Cycle);
+					Serializer.Append("B=");
+					Serializer.WriteValueInt64(Cycle);
 				}
 
-				Serializer.NextEventField();
-				Serializer.WriteInteger(bIsUnknownSpecId ? "?id" : "id", SpecId);
+				Serializer.Append(bIsUnknownSpecId ? " ?id=" : " id=");
+				Serializer.WriteValueInt64(SpecId);
 			}
 			else
 			{
@@ -1764,7 +1747,8 @@ void FConvertToTextAnalyzer::ConvertToTextEventBatchV2(const FOnEventContext& Co
 				//Serializer.WriteInteger("E", ActualCycle);
 				if (bIsAbsoluteTimestamp)
 				{
-					Serializer.WriteInteger("*E", Cycle);
+					Serializer.Append("*E=");
+					Serializer.WriteValueInt64(Cycle);
 					Serializer.AppendChar('(');
 					double Time = Context.EventTime.AsSeconds(Cycle);
 					Serializer.WriteValueTime(Time);
@@ -1772,7 +1756,8 @@ void FConvertToTextAnalyzer::ConvertToTextEventBatchV2(const FOnEventContext& Co
 				}
 				else
 				{
-					Serializer.WriteInteger("E", Cycle);
+					Serializer.Append("E=");
+					Serializer.WriteValueInt64(Cycle);
 				}
 			}
 		}
diff --git a/Engine/Source/Programs/TraceAnalyzer/Private/TextSerializer.h b/Engine/Source/Programs/TraceAnalyzer/Private/TextSerializer.h
--- a/Engine/Source/Programs/TraceAnalyzer/Private/TextSerializer.h
+++ b/Engine/Source/Programs/TraceAnalyzer/Private/TextSerializer.h
@@ -27,11 +27,37 @@ public:
 	virtual void Appendf(const char* Format, ...) = 0;
 	virtual bool Commit() = 0;
 
+	//////////////////////////////////////////////////
+	// Attributes
+
+	void BeginAttributeSet()
+	{
+		AttributeCount = 0;
+	}
+
+	void BeginAttribute()
+	{
+		if (AttributeCount == 0)
+		{
+			AppendChar('\t');
+		}
+		else
+		{
+			AppendChar(' ');
+		}
+		++AttributeCount;
+	}
+
+	void EndAttribute()
+	{
+	}
+
 	//////////////////////////////////////////////////
 	// NEW_EVENT
 
 	void BeginNewEventHeader()
 	{
+		BeginAttributeSet();
 		if (bWriteEventHeader)
 		{
 #if !UE_TRACE_ANALYSIS_DEBUG
@@ -42,12 +68,14 @@ public:
 			}
 #endif
 			bLastWasNewEvent = true;
-			Append("\tNEW_EVENT : ");
+			BeginAttribute();
+			Append("NEW_EVENT :");
+			EndAttribute();
 		}
 	}
 	void EndNewEventHeader()
 	{
-		if (bWriteEventHeader)
+		if (AttributeCount > 0)
 		{
 			AppendChar('\n');
 		}
@@ -58,14 +86,14 @@ public:
 	}
 	void BeginField()
 	{
+		BeginAttributeSet();
 		if (bWriteEventHeader)
 		{
-			Append("\t\tFIELD : ");
-		}
-		else
-		{
-			Append("\tFIELD : ");
+			AppendChar('\t');
 		}
+		BeginAttribute();
+		Append("FIELD :");
+		EndAttribute();
 	}
 	void EndField()
 	{
@@ -84,21 +112,20 @@ public:
 
 	void BeginEvent(uint32 CtxThreadId)
 	{
+		BeginAttributeSet();
 		if (bWriteEventHeader)
 		{
 			bLastWasNewEvent = false;
+			BeginAttribute();
 			if (CtxThreadId != (uint32)-1)
 			{
-				Appendf("\tEVENT [%u]", CtxThreadId);
+				Appendf("EVENT [%u]", CtxThreadId);
 			}
 			else
 			{
-				Append("\tEVENT");
+				Append("EVENT");
 			}
-		}
-		else
-		{
-			AppendChar('\t');
+			EndAttribute();
 		}
 	}
 
@@ -106,28 +133,20 @@ public:
 	{
 		if (bWriteEventHeader)
 		{
-			Appendf(" %s.%s", LoggerName, Name);
+			BeginAttribute();
+			Appendf("%s.%s :", LoggerName, Name);
+			EndAttribute();
 		}
 	}
 
-	void BeginEventFields()
+	void EndEvent()
 	{
-		if (bWriteEventHeader)
+		if (AttributeCount > 0)
 		{
-			Append(" : ");
+			AppendChar('\n');
 		}
 	}
 
-	void NextEventField()
-	{
-		AppendChar(' ');
-	}
-
-	void EndEvent()
-	{
-		AppendChar('\n');
-	}
-
 	//////////////////////////////////////////////////
 	// Array: [1 2 3...]
 
@@ -204,19 +223,20 @@ public:
 		AppendChar('=');
 	}
 
-	void WriteString(const ANSICHAR* Name, const ANSICHAR* Value)             { WriteKey(Name); WriteValueString(Value); }
-	void WriteString(const ANSICHAR* Name, const ANSICHAR* Value, uint32 Len) { WriteKey(Name); WriteValueString(Value, Len); }
-	void WriteBool(const ANSICHAR* Name, bool Value)                          { WriteKey(Name); WriteValueBool(Value); }
-	void WriteInteger(const ANSICHAR* Name, int64 Value)                      { WriteKey(Name); WriteValueInt64(Value); }
-	void WriteIntegerHex(const ANSICHAR* Name, int64 Value)                   { WriteKey(Name); WriteValueHex64(uint64(Value)); }
-	void WriteFloat(const ANSICHAR* Name, float Value)                        { WriteKey(Name); WriteValueFloat(Value); }
-	void WriteDouble(const ANSICHAR* Name, double Value)                      { WriteKey(Name); WriteValueDouble(Value); }
-	void WriteNull(const ANSICHAR* Name)                                      { WriteKey(Name); WriteValueNull(); }
-	void WriteBinary(const ANSICHAR* Name, const void* Data, uint32 Size)     { WriteKey(Name); WriteValueBinary(Data, Size); }
+	void WriteAttributeString(const ANSICHAR* Name, const ANSICHAR* Value)             { BeginAttribute(); WriteKey(Name); WriteValueString(Value);        EndAttribute(); }
+	void WriteAttributeString(const ANSICHAR* Name, const ANSICHAR* Value, uint32 Len) { BeginAttribute(); WriteKey(Name); WriteValueString(Value, Len);   EndAttribute(); }
+	void WriteAttributeBool(const ANSICHAR* Name, bool Value)                          { BeginAttribute(); WriteKey(Name); WriteValueBool(Value);          EndAttribute(); }
+	void WriteAttributeInteger(const ANSICHAR* Name, int64 Value)                      { BeginAttribute(); WriteKey(Name); WriteValueInt64(Value);         EndAttribute(); }
+	void WriteAttributeIntegerHex(const ANSICHAR* Name, int64 Value)                   { BeginAttribute(); WriteKey(Name); WriteValueHex64(uint64(Value)); EndAttribute(); }
+	void WriteAttributeFloat(const ANSICHAR* Name, float Value)                        { BeginAttribute(); WriteKey(Name); WriteValueFloat(Value);         EndAttribute(); }
+	void WriteAttributeDouble(const ANSICHAR* Name, double Value)                      { BeginAttribute(); WriteKey(Name); WriteValueDouble(Value);        EndAttribute(); }
+	void WriteAttributeNull(const ANSICHAR* Name)                                      { BeginAttribute(); WriteKey(Name); WriteValueNull();               EndAttribute(); }
+	void WriteAttributeBinary(const ANSICHAR* Name, const void* Data, uint32 Size)     { BeginAttribute(); WriteKey(Name); WriteValueBinary(Data, Size);   EndAttribute(); }
 
 protected:
 	bool bWriteEventHeader = true;
 	bool bLastWasNewEvent = false;
+	uint32 AttributeCount = 0;
 };
 
 ////////////////////////////////////////////////////////////////////////////////////////////////////
