[TraceAnalyzer]
* Added display of the actual element type for numeric fields (integer and floating point; scalar and arrays).
* Fixed handling (convertion to text) for array fields (for both integer and floating point arrays).

#Synced-CL 49330596

diff --git a/Engine/Source/Programs/TraceAnalyzer/Private/ConvertToText.cpp b/Engine/Source/Programs/TraceAnalyzer/Private/ConvertToText.cpp
--- a/Engine/Source/Programs/TraceAnalyzer/Private/ConvertToText.cpp
+++ b/Engine/Source/Programs/TraceAnalyzer/Private/ConvertToText.cpp
@@ -15,6 +15,8 @@
 #include "Io.h"
 #include "TextSerializer.h"
 
+#include <type_traits>
+
 namespace UE
 {
 namespace TraceAnalyzer
@@ -77,8 +79,13 @@ private:
 	virtual void OnVersion(uint32 InTransportVersion, uint32 InProtocolVersion) override;
 	virtual bool OnNewEvent(uint16 RouteId, const FEventTypeInfo& TypeInfo) override;
 	virtual bool OnEvent(uint16 RouteId, EStyle, const FOnEventContext& Context) override;
-	uint32 GetTotalEventSize(EStyle Style, const FOnEventContext& Context) const;
 	void DefineTimer(uint32 SpecId, const TCHAR* TimerName);
+
+	void WriteArrayHex(FTextSerializer& Serializer, TConstArrayView<uint8> ArrayView, uint32 MaxCount) const;
+
+	template <typename ValueType>
+	void WriteArray(FTextSerializer& Serializer, const TArrayReader<ValueType>& Reader, uint32 MaxCount) const;
+
 	void ConvertToTextEventBatch(const FOnEventContext& Context, FTextSerializer& Serializer, const uint8* BufferPtr, uint32 BufferSize);
 	void ConvertToTextEventBatchV2(const FOnEventContext& Context, FTextSerializer& Serializer, const uint8* BufferPtr, uint32 BufferSize);
 	void PrintUnknownTimers(FTextSerializer& Serializer);
@@ -1042,44 +1049,93 @@ bool FConvertToTextAnalyzer::OnNewEvent(uint16 RouteId, const FEventTypeInfo& Ty
 		/** The size of the field's data in bytes. */
 		Serializer.WriteAttributeInteger("Size", FieldInfo.GetSize());
 
+		/** Is this field signed? */
+		Serializer.WriteAttributeBool("IsSigned", FieldInfo.IsSigned());
+
+		/** Is this field an array-type field? */
+		Serializer.WriteAttributeBool("IsArray", FieldInfo.IsArray());
+
 		Serializer.BeginAttribute();
 		/** What type of field is this? */
 		Serializer.WriteKey("Type");
+
+		using EType = UE::Trace::IAnalyzer::FEventFieldInfo::EType;
 		switch (FieldInfo.GetType())
 		{
-		case IAnalyzer::FEventFieldInfo::EType::Integer: // int64
+		case EType::Integer:
+		{
 			Serializer.Append("Integer");
+			if (FieldInfo.IsSigned())
+			{
+				switch (FieldInfo.GetTypeSize())
+				{
+				case 1: // int8
+				case 2: // int16
+				case 4: // int32
+				case 8: // int64
+					Serializer.Appendf(" (int%d%s)", 8 * FieldInfo.GetTypeSize(), FieldInfo.IsArray() ? "[]" : "");
+					break;
+				default:
+					Serializer.Appendf(" (%d bytes/element)", FieldInfo.GetTypeSize());
+				}
+			}
+			else
+			{
+				switch (FieldInfo.GetTypeSize())
+				{
+				case 1: // uint8
+				case 2: // uint16
+				case 4: // uint32
+				case 8: // uint64
+					Serializer.Appendf(" (uint%d%s)", 8 * FieldInfo.GetTypeSize(), FieldInfo.IsArray() ? "[]" : "");
+					break;
+				default:
+					Serializer.Appendf(" (%d bytes/element)", FieldInfo.GetTypeSize());
+				}
+			}
 			break;
-		case IAnalyzer::FEventFieldInfo::EType::Float: // double
-			Serializer.Append("Float");
+		}
+		case EType::Float: // Float32 or Float64
+		{
+			Serializer.Appendf("Float");
+			switch (FieldInfo.GetTypeSize())
+			{
+			case 4: // float
+				Serializer.Appendf(" (float%s)", FieldInfo.IsArray() ? "[]" : "");
+				break;
+			case 8: // double
+				Serializer.Appendf(" (double%s)", FieldInfo.IsArray() ? "[]" : "");
+				break;
+			default:
+				Serializer.Appendf(" (%d bytes/element)", FieldInfo.GetTypeSize());
+			}
 			break;
-		case IAnalyzer::FEventFieldInfo::EType::AnsiString:
+		}
+		case EType::AnsiString:
 			Serializer.Append("AnsiString");
 			break;
-		case IAnalyzer::FEventFieldInfo::EType::WideString:
+		case EType::WideString:
 			Serializer.Append("WideString");
 			break;
-		case IAnalyzer::FEventFieldInfo::EType::Reference8:
+		case EType::Reference8:
 			Serializer.Append("Reference8");
 			break;
-		case IAnalyzer::FEventFieldInfo::EType::Reference16:
+		case EType::Reference16:
 			Serializer.Append("Reference16");
 			break;
-		case IAnalyzer::FEventFieldInfo::EType::Reference32:
+		case EType::Reference32:
 			Serializer.Append("Reference32");
 			break;
-		case IAnalyzer::FEventFieldInfo::EType::Reference64:
+		case EType::Reference64:
 			Serializer.Append("Reference64");
 			break;
 		default:
 			Serializer.WriteValueUInt32(static_cast<uint32>(FieldInfo.GetType()));
+			Serializer.Appendf(" (%d bytes/element)", FieldInfo.GetTypeSize());
 			break;
 		}
 		Serializer.EndAttribute();
 
-		/** Is this field an array-type field? */
-		Serializer.WriteAttributeBool("IsArray", FieldInfo.IsArray());
-
 		Serializer.EndField();
 	}
 
@@ -1251,13 +1307,15 @@ bool FConvertToTextAnalyzer::OnEvent(uint16 RouteId, EStyle Style, const FOnEven
 			}
 		}
 
+		using EType = UE::Trace::IAnalyzer::FEventFieldInfo::EType;
+
 		Serializer.BeginAttribute();
 		Serializer.WriteKey(FieldInfo.GetName());
 		if (FieldInfo.IsArray())
 		{
 			switch (FieldInfo.GetType())
 			{
-				case FEventFieldInfo::EType::Integer: // int64
+				case EType::Integer:
 				{
 #if 0
 					// Probe element size...
@@ -1293,101 +1351,102 @@ bool FConvertToTextAnalyzer::OnEvent(uint16 RouteId, EStyle Style, const FOnEven
 
 					if (FieldInfo.IsSigned())
 					{
-						Serializer.BeginArray();
-						const TArrayReader<int64>& Reader = EventData.GetArray<int64>(FieldInfo.GetName());
-						const uint32 ArrayCount = Reader.Num();
-						const uint32 ActualCount = FMath::Min(ArrayCount, 8u);
-						for (uint32 ArrayIndex = 0; ArrayIndex < ActualCount; ++ArrayIndex)
+						switch (FieldInfo.GetTypeSize())
 						{
-							if (ArrayIndex != 0)
-							{
-								Serializer.NextArrayElement();
-							}
-							const int64 Value = Reader[ArrayIndex];
-							Serializer.WriteValueInt64Auto(Value);
+						case 1: // int8
+						{
+							const TArrayReader<int8>& Reader = EventData.GetArray<int8>(FieldInfo.GetName());
+							WriteArray<int8>(Serializer, Reader, 64);
+							break;
+						}
+						case 2: // int16
+						{
+							const TArrayReader<int16>& Reader = EventData.GetArray<int16>(FieldInfo.GetName());
+							WriteArray<int16>(Serializer, Reader, 32);
+							break;
+						}
+						case 4: // int32
+						{
+							const TArrayReader<int32>& Reader = EventData.GetArray<int32>(FieldInfo.GetName());
+							WriteArray<int32>(Serializer, Reader, 16);
+							break;
 						}
-						if (ActualCount != ArrayCount)
+						case 8: // int64
 						{
-							Serializer.Appendf(" ... | %u elements", ArrayCount);
+							const TArrayReader<int64>& Reader = EventData.GetArray<int64>(FieldInfo.GetName());
+							WriteArray<int64>(Serializer, Reader, 8);
+							break;
 						}
-						Serializer.EndArray();
+						default:
+							const TArrayReader<uint8>& Reader = EventData.GetArray<uint8>(FieldInfo.GetName());
+							WriteArrayHex(Serializer, TConstArrayView<uint8>((const uint8*)Reader.GetData(), (int32)Reader.Num()), 32);
+							break;
+						} // switch (FieldInfo.GetTypeSize())
+
 					}
 					else // unsigned
 					{
-						const TArrayReader<uint8>& Reader8 = EventData.GetArray<uint8>(FieldInfo.GetName());
-						const uint8* Data8 = Reader8.GetData();
-						if (Data8)
+						switch (FieldInfo.GetTypeSize())
 						{
-							// Write array as 0x[00 01 ... FF]
-							Serializer.Append("0x");
-							Serializer.BeginArray();
-							const uint32 ArrayCount = Reader8.Num();
-							const uint32 ActualCount = FMath::Min(ArrayCount, 32u);
-							for (uint32 ArrayIndex = 0; ArrayIndex < ActualCount; ++ArrayIndex)
-							{
-								if (ArrayIndex != 0)
-								{
-									Serializer.NextArrayElement();
-								}
-								const uint8 Value8 = Reader8[ArrayIndex];
-								const uint32 Value = (uint32)Data8[ArrayIndex];
-								check(Value8 == Data8[ArrayIndex]);
-								Serializer.Appendf("%02X", Value);
-							}
-							if (ActualCount != ArrayCount)
-							{
-								Serializer.Appendf(" ... | %u bytes", ArrayCount);
-							}
-							Serializer.EndArray();
+						case 1: // uint8
+						{
+							const TArrayReader<uint8>& Reader = EventData.GetArray<uint8>(FieldInfo.GetName());
+							//WriteArray<uint8>(Serializer, Reader, 64);
+							WriteArrayHex(Serializer, TConstArrayView<uint8>((const uint8*)Reader.GetData(), (int32)Reader.Num()), 32);
 							break;
 						}
-
-						Serializer.BeginArray();
-						const TArrayReader<uint64>& Reader = EventData.GetArray<uint64>(FieldInfo.GetName());
-						const uint32 ArrayCount = Reader.Num();
-						const uint32 ActualCount = FMath::Min(ArrayCount, 8u);
-						for (uint32 ArrayIndex = 0; ArrayIndex < ActualCount; ++ArrayIndex)
+						case 2: // uint16
 						{
-							if (ArrayIndex != 0)
-							{
-								Serializer.NextArrayElement();
-							}
-							const uint64 Value = Reader[ArrayIndex];
-							Serializer.WriteValueUInt64Auto(Value);
+							const TArrayReader<uint16>& Reader = EventData.GetArray<uint16>(FieldInfo.GetName());
+							WriteArray<uint16>(Serializer, Reader, 32);
+							break;
+						}
+						case 4: // uint32
+						{
+							const TArrayReader<uint32>& Reader = EventData.GetArray<uint32>(FieldInfo.GetName());
+							WriteArray<uint32>(Serializer, Reader, 16);
+							break;
 						}
-						if (ActualCount != ArrayCount)
+						case 8: // uint64
 						{
-							Serializer.Appendf(" ... | %u elements", ArrayCount);
+							const TArrayReader<uint64>& Reader = EventData.GetArray<uint64>(FieldInfo.GetName());
+							WriteArray<uint64>(Serializer, Reader, 8);
+							break;
 						}
-						Serializer.EndArray();
+						default:
+							const TArrayReader<uint8>& Reader = EventData.GetArray<uint8>(FieldInfo.GetName());
+							WriteArrayHex(Serializer, TConstArrayView<uint8>((const uint8*)Reader.GetData(), (int32)Reader.Num()), 32);
+							break;
+						} // switch (FieldInfo.GetTypeSize())
 					}
 					break;
 				}
 
-				case FEventFieldInfo::EType::Float: // double
+				case EType::Float: // Float32 or Float64
 				{
-					Serializer.BeginArray();
-					const TArrayReader<double>& Reader = EventData.GetArray<double>(FieldInfo.GetName());
-					const uint32 ArrayCount = Reader.Num();
-					const uint32 ActualCount = FMath::Min(ArrayCount, 4u);
-					for (uint32 ArrayIndex = 0; ArrayIndex < ActualCount; ++ArrayIndex)
+					switch (FieldInfo.GetTypeSize())
 					{
-						if (ArrayIndex != 0)
-						{
-							Serializer.NextArrayElement();
-						}
-						const double Value = Reader[ArrayIndex];
-						Serializer.WriteValueDouble(Value);
+					case 4: // float
+					{
+						const TArrayReader<float>& Reader = EventData.GetArray<float>(FieldInfo.GetName());
+						WriteArray<float>(Serializer, Reader, 8);
+						break;
 					}
-					if (ActualCount != ArrayCount)
+					case 8: // double
 					{
-						Serializer.Appendf(" ... | %u elements", ArrayCount);
+						const TArrayReader<double>& Reader = EventData.GetArray<double>(FieldInfo.GetName());
+						WriteArray<double>(Serializer, Reader, 8);
+						break;
 					}
-					Serializer.EndArray();
+					default:
+						const TArrayReader<uint8>& Reader = EventData.GetArray<uint8>(FieldInfo.GetName());
+						WriteArrayHex(Serializer, TConstArrayView<uint8>((const uint8*)Reader.GetData(), (int32)Reader.Num()), 32);
+						break;
+					} // switch (FieldInfo.GetTypeSize())
 					break;
 				}
 
-				case FEventFieldInfo::EType::AnsiString:
+				case EType::AnsiString:
 				{
 					FAnsiStringView Value;
 					EventData.GetString(FieldInfo.GetName(), Value);
@@ -1395,7 +1454,7 @@ bool FConvertToTextAnalyzer::OnEvent(uint16 RouteId, EStyle Style, const FOnEven
 					break;
 				}
 
-				case FEventFieldInfo::EType::WideString:
+				case EType::WideString:
 				{
 					FString Value;
 					EventData.GetString(FieldInfo.GetName(), Value);
@@ -1414,7 +1473,7 @@ bool FConvertToTextAnalyzer::OnEvent(uint16 RouteId, EStyle Style, const FOnEven
 		{
 			switch (FieldInfo.GetType())
 			{
-				case FEventFieldInfo::EType::Integer:
+				case EType::Integer:
 				{
 					if (FieldInfo.IsSigned())
 					{
@@ -1451,14 +1510,14 @@ bool FConvertToTextAnalyzer::OnEvent(uint16 RouteId, EStyle Style, const FOnEven
 					break;
 				}
 
-				case FEventFieldInfo::EType::Float: // double
+				case EType::Float: // double
 				{
 					const double Value = EventData.GetValue<double>(FieldInfo.GetName());
 					Serializer.WriteValueDouble(Value);
 					break;
 				}
 
-				case FEventFieldInfo::EType::Reference8:
+				case EType::Reference8:
 				{
 					check(FieldInfo.GetSize() == 1);
 					UE::Trace::FEventRef8 RefValue = EventData.GetReferenceValue<uint8>(FieldIndex);
@@ -1466,7 +1525,7 @@ bool FConvertToTextAnalyzer::OnEvent(uint16 RouteId, EStyle Style, const FOnEven
 					break;
 				}
 
-				case FEventFieldInfo::EType::Reference16:
+				case EType::Reference16:
 				{
 					check(FieldInfo.GetSize() == 2);
 					UE::Trace::FEventRef16 RefValue = EventData.GetReferenceValue<uint16>(FieldIndex);
@@ -1474,7 +1533,7 @@ bool FConvertToTextAnalyzer::OnEvent(uint16 RouteId, EStyle Style, const FOnEven
 					break;
 				}
 
-				case FEventFieldInfo::EType::Reference32:
+				case EType::Reference32:
 				{
 					check(FieldInfo.GetSize() == 4);
 					UE::Trace::FEventRef32 RefValue = EventData.GetReferenceValue<uint32>(FieldIndex);
@@ -1482,7 +1541,7 @@ bool FConvertToTextAnalyzer::OnEvent(uint16 RouteId, EStyle Style, const FOnEven
 					break;
 				}
 
-				case FEventFieldInfo::EType::Reference64:
+				case EType::Reference64:
 				{
 					check(FieldInfo.GetSize() == 8);
 					UE::Trace::FEventRef64 RefValue = EventData.GetReferenceValue<uint64>(FieldIndex);
@@ -1542,6 +1601,74 @@ bool FConvertToTextAnalyzer::OnEvent(uint16 RouteId, EStyle Style, const FOnEven
 
 ////////////////////////////////////////////////////////////////////////////////////////////////////
 
+void FConvertToTextAnalyzer::WriteArrayHex(FTextSerializer& Serializer, TConstArrayView<uint8> ArrayView, uint32 MaxCount) const
+{
+	// Write array as 0x[00 01 ... FF]
+	Serializer.Append("0x");
+	Serializer.BeginArray();
+
+	const uint32 ArrayCount = ArrayView.Num();
+	const uint32 ActualCount = FMath::Min(ArrayCount, MaxCount);
+	for (uint32 ArrayIndex = 0; ArrayIndex < ActualCount; ++ArrayIndex)
+	{
+		if (ArrayIndex != 0)
+		{
+			Serializer.NextArrayElement();
+		}
+		const uint8 Value = ArrayView.GetData()[ArrayIndex];
+		Serializer.Appendf("%02X", uint32(Value));
+	}
+	if (ActualCount != ArrayCount)
+	{
+		Serializer.Appendf(" ... | %u bytes", ArrayCount);
+	}
+
+	Serializer.EndArray();
+}
+
+////////////////////////////////////////////////////////////////////////////////////////////////////
+
+template <typename ValueType>
+void FConvertToTextAnalyzer::WriteArray(FTextSerializer& Serializer, const TArrayReader<ValueType>& Reader, uint32 MaxCount) const
+{
+	TConstArrayView<ValueType> ArrayView(Reader.GetData(), Reader.Num());
+
+	Serializer.BeginArray();
+
+	const uint32 ArrayCount = Reader.Num();
+	const uint32 ActualCount = FMath::Min(ArrayCount, MaxCount);
+	for (uint32 ArrayIndex = 0; ArrayIndex < ActualCount; ++ArrayIndex)
+	{
+		if (ArrayIndex != 0)
+		{
+			Serializer.NextArrayElement();
+		}
+
+		const ValueType Value = ArrayView[ArrayIndex];
+		const ValueType Value2 = Reader[ArrayIndex];
+		check(Value == Value2);
+
+		     if constexpr (std::is_same_v<ValueType, uint8>)  { Serializer.WriteValueUInt8(Value); }
+		else if constexpr (std::is_same_v<ValueType, uint16>) { Serializer.WriteValueUInt16(Value); }
+		else if constexpr (std::is_same_v<ValueType, uint32>) { Serializer.WriteValueUInt64Auto(uint64(Value)); }
+		else if constexpr (std::is_same_v<ValueType, uint64>) { Serializer.WriteValueUInt64Auto(Value); }
+		else if constexpr (std::is_same_v<ValueType, int8>)   { Serializer.WriteValueInt8(Value); }
+		else if constexpr (std::is_same_v<ValueType, int16>)  { Serializer.WriteValueInt16(Value); }
+		else if constexpr (std::is_same_v<ValueType, int32>)  { Serializer.WriteValueInt64Auto(int64(Value)); }
+		else if constexpr (std::is_same_v<ValueType, int64>)  { Serializer.WriteValueInt64Auto(Value); }
+		else if constexpr (std::is_same_v<ValueType, float>)  { Serializer.WriteValueFloatAuto(Value); }
+		else if constexpr (std::is_same_v<ValueType, double>) { Serializer.WriteValueDoubleAuto(Value); }
+	}
+	if (ActualCount != ArrayCount)
+	{
+		Serializer.Appendf(" ... | %u elements", ArrayCount);
+	}
+
+	Serializer.EndArray();
+}
+
+////////////////////////////////////////////////////////////////////////////////////////////////////
+
 void FConvertToTextAnalyzer::ConvertToTextEventBatch(const FOnEventContext& Context, FTextSerializer& Serializer, const uint8* BufferPtr, uint32 BufferSize)
 {
 	++NumCpuBatches;
diff --git a/Engine/Source/Programs/TraceAnalyzer/Private/TextSerializer.cpp b/Engine/Source/Programs/TraceAnalyzer/Private/TextSerializer.cpp
--- a/Engine/Source/Programs/TraceAnalyzer/Private/TextSerializer.cpp
+++ b/Engine/Source/Programs/TraceAnalyzer/Private/TextSerializer.cpp
@@ -6,6 +6,9 @@
 #include "Math/NumericLimits.h"
 #include "Misc/CString.h"
 
+#include <limits>
+#include <cmath>
+
 namespace UE
 {
 namespace TraceAnalyzer
@@ -86,6 +89,82 @@ void FTextSerializer::WriteValueUInt64Auto(uint64 Value)
 
 ////////////////////////////////////////////////////////////////////////////////////////////////////
 
+void FTextSerializer::WriteValueFloatAuto(float Value)
+{
+	if (Value == std::numeric_limits<float>::min())
+	{
+		Append("MIN_flt");
+	}
+	else if (Value == std::numeric_limits<float>::max())
+	{
+		Append("MAX_flt");
+	}
+	else if (Value == std::numeric_limits<float>::lowest())
+	{
+		Append("LOWEST_flt");
+	}
+	else if (Value == std::numeric_limits<float>::denorm_min())
+	{
+		Append("DENORM_MIN_flt");
+	}
+	else if (Value == -std::numeric_limits<float>::infinity())
+	{
+		Append("-INF_flt");
+	}
+	else if (Value == std::numeric_limits<float>::infinity())
+	{
+		Append("INF_flt");
+	}
+	else if (std::isnan(Value))
+	{
+		Append("NAN_flt");
+	}
+	else
+	{
+		Appendf("%g", Value);
+	}
+}
+
+////////////////////////////////////////////////////////////////////////////////////////////////////
+
+void FTextSerializer::WriteValueDoubleAuto(double Value)
+{
+	if (Value == std::numeric_limits<double>::min())
+	{
+		Append("MIN_dbl");
+	}
+	else if (Value == std::numeric_limits<double>::max())
+	{
+		Append("MAX_dbl");
+	}
+	else if (Value == std::numeric_limits<double>::lowest())
+	{
+		Append("LOWEST_dbl");
+	}
+	else if (Value == std::numeric_limits<double>::denorm_min())
+	{
+		Append("DENORM_MIN_dbl");
+	}
+	else if (Value == -std::numeric_limits<double>::infinity())
+	{
+		Append("-INF_dbl");
+	}
+	else if (Value == std::numeric_limits<double>::infinity())
+	{
+		Append("INF_dbl");
+	}
+	else if (std::isnan(Value))
+	{
+		Append("NAN_dbl");
+	}
+	else
+	{
+		Appendf("%g", Value);
+	}
+}
+
+////////////////////////////////////////////////////////////////////////////////////////////////////
+
 void FTextSerializer::WriteValueBinary(const void* Data, uint32 Size)
 {
 	AppendChar('<');
diff --git a/Engine/Source/Programs/TraceAnalyzer/Private/TextSerializer.h b/Engine/Source/Programs/TraceAnalyzer/Private/TextSerializer.h
--- a/Engine/Source/Programs/TraceAnalyzer/Private/TextSerializer.h
+++ b/Engine/Source/Programs/TraceAnalyzer/Private/TextSerializer.h
@@ -209,6 +209,9 @@ public:
 	void WriteValueFloat(float Value)       { Appendf("%f", Value); }
 	void WriteValueDouble(double Value)     { Appendf("%f", Value); }
 
+	void WriteValueFloatAuto(float Value);
+	void WriteValueDoubleAuto(double Value);
+
 	void WriteValueTime(double Time)        { Appendf("%f", Time); }
 	void WriteValueNull()                   { Append("null"); }
 
