[TraceAnalyzer]
* Fixed support for Mac and Linux.
* Improved formatting for float and double values.

#Synced-CL 49645631

diff --git a/Engine/Source/Programs/TraceAnalyzer/Private/Common.h b/Engine/Source/Programs/TraceAnalyzer/Private/Common.h
--- a/Engine/Source/Programs/TraceAnalyzer/Private/Common.h
+++ b/Engine/Source/Programs/TraceAnalyzer/Private/Common.h
@@ -2,7 +2,12 @@
 
 #pragma once
 
+#include "CoreTypes.h"
+
+// TraceAnalyzer
 #include "Io.h"
+
+// TraceAnalysis
 #include "Trace/DataStream.h"
 #include "Trace/StoreClient.h"
 
diff --git a/Engine/Source/Programs/TraceAnalyzer/Private/ConvertToText.cpp b/Engine/Source/Programs/TraceAnalyzer/Private/ConvertToText.cpp
--- a/Engine/Source/Programs/TraceAnalyzer/Private/ConvertToText.cpp
+++ b/Engine/Source/Programs/TraceAnalyzer/Private/ConvertToText.cpp
@@ -7,13 +7,17 @@
 #include "Containers/UnrealString.h"
 #include "HAL/PlatformTime.h"
 #include "Templates/UniquePtr.h"
+
+// TraceAnalysis
 #include "Trace/Analysis.h"
 #include "Trace/Analyzer.h"
 #include "Trace/DataStream.h"
 
+// TraceAnalyzer
 #include "Command.h"
 #include "Io.h"
 #include "TextSerializer.h"
+#include "TraceAnalyzer.h"
 
 #include <type_traits>
 
@@ -1353,70 +1357,71 @@ bool FConvertToTextAnalyzer::OnEvent(uint16 RouteId, EStyle Style, const FOnEven
 					{
 						switch (FieldInfo.GetTypeSize())
 						{
-						case 1: // int8
-						{
-							const TArrayReader<int8>& Reader = EventData.GetArray<int8>(FieldInfo.GetName());
-							WriteArray<int8>(Serializer, Reader, 64);
-							break;
-						}
-						case 2: // int16
-						{
-							const TArrayReader<int16>& Reader = EventData.GetArray<int16>(FieldInfo.GetName());
-							WriteArray<int16>(Serializer, Reader, 32);
-							break;
-						}
-						case 4: // int32
-						{
-							const TArrayReader<int32>& Reader = EventData.GetArray<int32>(FieldInfo.GetName());
-							WriteArray<int32>(Serializer, Reader, 16);
-							break;
-						}
-						case 8: // int64
-						{
-							const TArrayReader<int64>& Reader = EventData.GetArray<int64>(FieldInfo.GetName());
-							WriteArray<int64>(Serializer, Reader, 8);
-							break;
-						}
-						default:
-							const TArrayReader<uint8>& Reader = EventData.GetArray<uint8>(FieldInfo.GetName());
-							WriteArrayHex(Serializer, TConstArrayView<uint8>((const uint8*)Reader.GetData(), (int32)Reader.Num()), 32);
-							break;
+							case 1: // int8
+							{
+								const TArrayReader<int8>& Reader = EventData.GetArray<int8>(FieldInfo.GetName());
+								WriteArray<int8>(Serializer, Reader, 64);
+								break;
+							}
+							case 2: // int16
+							{
+								const TArrayReader<int16>& Reader = EventData.GetArray<int16>(FieldInfo.GetName());
+								WriteArray<int16>(Serializer, Reader, 32);
+								break;
+							}
+							case 4: // int32
+							{
+								const TArrayReader<int32>& Reader = EventData.GetArray<int32>(FieldInfo.GetName());
+								WriteArray<int32>(Serializer, Reader, 16);
+								break;
+							}
+							case 8: // int64
+							{
+								const TArrayReader<int64>& Reader = EventData.GetArray<int64>(FieldInfo.GetName());
+								WriteArray<int64>(Serializer, Reader, 8);
+								break;
+							}
+							default:
+							{
+								const TArrayReader<uint8>& Reader = EventData.GetArray<uint8>(FieldInfo.GetName());
+								WriteArrayHex(Serializer, TConstArrayView<uint8>((const uint8*)Reader.GetData(), (int32)Reader.Num()), 32);
+							}
 						} // switch (FieldInfo.GetTypeSize())
-
 					}
 					else // unsigned
 					{
 						switch (FieldInfo.GetTypeSize())
 						{
-						case 1: // uint8
-						{
-							const TArrayReader<uint8>& Reader = EventData.GetArray<uint8>(FieldInfo.GetName());
-							//WriteArray<uint8>(Serializer, Reader, 64);
-							WriteArrayHex(Serializer, TConstArrayView<uint8>((const uint8*)Reader.GetData(), (int32)Reader.Num()), 32);
-							break;
-						}
-						case 2: // uint16
-						{
-							const TArrayReader<uint16>& Reader = EventData.GetArray<uint16>(FieldInfo.GetName());
-							WriteArray<uint16>(Serializer, Reader, 32);
-							break;
-						}
-						case 4: // uint32
-						{
-							const TArrayReader<uint32>& Reader = EventData.GetArray<uint32>(FieldInfo.GetName());
-							WriteArray<uint32>(Serializer, Reader, 16);
-							break;
-						}
-						case 8: // uint64
-						{
-							const TArrayReader<uint64>& Reader = EventData.GetArray<uint64>(FieldInfo.GetName());
-							WriteArray<uint64>(Serializer, Reader, 8);
-							break;
-						}
-						default:
-							const TArrayReader<uint8>& Reader = EventData.GetArray<uint8>(FieldInfo.GetName());
-							WriteArrayHex(Serializer, TConstArrayView<uint8>((const uint8*)Reader.GetData(), (int32)Reader.Num()), 32);
-							break;
+							case 1: // uint8
+							{
+								const TArrayReader<uint8>& Reader = EventData.GetArray<uint8>(FieldInfo.GetName());
+								//WriteArray<uint8>(Serializer, Reader, 64);
+								WriteArrayHex(Serializer, TConstArrayView<uint8>((const uint8*)Reader.GetData(), (int32)Reader.Num()), 32);
+								break;
+							}
+							case 2: // uint16
+							{
+								const TArrayReader<uint16>& Reader = EventData.GetArray<uint16>(FieldInfo.GetName());
+								WriteArray<uint16>(Serializer, Reader, 32);
+								break;
+							}
+							case 4: // uint32
+							{
+								const TArrayReader<uint32>& Reader = EventData.GetArray<uint32>(FieldInfo.GetName());
+								WriteArray<uint32>(Serializer, Reader, 16);
+								break;
+							}
+							case 8: // uint64
+							{
+								const TArrayReader<uint64>& Reader = EventData.GetArray<uint64>(FieldInfo.GetName());
+								WriteArray<uint64>(Serializer, Reader, 8);
+								break;
+							}
+							default:
+							{
+								const TArrayReader<uint8>& Reader = EventData.GetArray<uint8>(FieldInfo.GetName());
+								WriteArrayHex(Serializer, TConstArrayView<uint8>((const uint8*)Reader.GetData(), (int32)Reader.Num()), 32);
+							}
 						} // switch (FieldInfo.GetTypeSize())
 					}
 					break;
@@ -1426,22 +1431,23 @@ bool FConvertToTextAnalyzer::OnEvent(uint16 RouteId, EStyle Style, const FOnEven
 				{
 					switch (FieldInfo.GetTypeSize())
 					{
-					case 4: // float
-					{
-						const TArrayReader<float>& Reader = EventData.GetArray<float>(FieldInfo.GetName());
-						WriteArray<float>(Serializer, Reader, 8);
-						break;
-					}
-					case 8: // double
-					{
-						const TArrayReader<double>& Reader = EventData.GetArray<double>(FieldInfo.GetName());
-						WriteArray<double>(Serializer, Reader, 8);
-						break;
-					}
-					default:
-						const TArrayReader<uint8>& Reader = EventData.GetArray<uint8>(FieldInfo.GetName());
-						WriteArrayHex(Serializer, TConstArrayView<uint8>((const uint8*)Reader.GetData(), (int32)Reader.Num()), 32);
-						break;
+						case 4: // float
+						{
+							const TArrayReader<float>& Reader = EventData.GetArray<float>(FieldInfo.GetName());
+							WriteArray<float>(Serializer, Reader, 8);
+							break;
+						}
+						case 8: // double
+						{
+							const TArrayReader<double>& Reader = EventData.GetArray<double>(FieldInfo.GetName());
+							WriteArray<double>(Serializer, Reader, 8);
+							break;
+						}
+						default:
+						{
+							const TArrayReader<uint8>& Reader = EventData.GetArray<uint8>(FieldInfo.GetName());
+							WriteArrayHex(Serializer, TConstArrayView<uint8>((const uint8*)Reader.GetData(), (int32)Reader.Num()), 32);
+						}
 					} // switch (FieldInfo.GetTypeSize())
 					break;
 				}
@@ -1510,10 +1516,28 @@ bool FConvertToTextAnalyzer::OnEvent(uint16 RouteId, EStyle Style, const FOnEven
 					break;
 				}
 
-				case EType::Float: // double
+				case EType::Float: // Float32 or Float64
 				{
-					const double Value = EventData.GetValue<double>(FieldInfo.GetName());
-					Serializer.WriteValueDouble(Value);
+					switch (FieldInfo.GetTypeSize())
+					{
+						case 4: // float
+						{
+							const float Value = EventData.GetValue<float>(FieldInfo.GetName());
+							Serializer.WriteValueFloatAuto(Value);
+							break;
+						}
+						case 8: // double
+						{
+							const double Value = EventData.GetValue<double>(FieldInfo.GetName());
+							Serializer.WriteValueDoubleAuto(Value);
+							break;
+						}
+						default:
+						{
+							// error
+							Serializer.Append("<UnknownFloatTypeSize>");
+						}
+					} // switch (FieldInfo.GetTypeSize())
 					break;
 				}
 
@@ -1977,7 +2001,7 @@ static int32 ConvertToTextMain(int32 ArgC, TCHAR const* const* ArgV)
 	Input = OpenFile(TraceFileName, false);
 	if (Input < 0)
 	{
-		wprintf(TEXT("Error: Cannot open the input trace file ('%s')!"), TraceFileName);
+		LogTraceAnalyzerError("Cannot open the input trace file ('%s')!", TraceFileName);
 		return -2;
 	}
 
@@ -1989,7 +2013,7 @@ static int32 ConvertToTextMain(int32 ArgC, TCHAR const* const* ArgV)
 			Output = OpenFile(OutputFileName, true);
 			if (Output < 0)
 			{
-				wprintf(TEXT("Warning: Cannot open the output text file ('%s')!\n"), OutputFileName);
+				LogTraceAnalyzerError("Cannot open the output text file('%s')!", OutputFileName);
 				CloseFile(Input);
 				return -3;
 			}
@@ -2012,7 +2036,7 @@ static int32 ConvertToTextMain(int32 ArgC, TCHAR const* const* ArgV)
 		}
 		else
 		{
-			wprintf(TEXT("Warning: Unknown cmd line argument '%s'!\n"), ArgV[ArgIndex]);
+			LogTraceAnalyzerWarning("Unknown cmd line argument '%s'!", ArgV[ArgIndex]);
 		}
 	}
 
diff --git a/Engine/Source/Programs/TraceAnalyzer/Private/Io.h b/Engine/Source/Programs/TraceAnalyzer/Private/Io.h
--- a/Engine/Source/Programs/TraceAnalyzer/Private/Io.h
+++ b/Engine/Source/Programs/TraceAnalyzer/Private/Io.h
@@ -2,6 +2,8 @@
 
 #pragma once
 
+#include "CoreTypes.h"
+
 #if PLATFORM_DESKTOP
 
 #if PLATFORM_WINDOWS
diff --git a/Engine/Source/Programs/TraceAnalyzer/Private/TextSerializer.cpp b/Engine/Source/Programs/TraceAnalyzer/Private/TextSerializer.cpp
--- a/Engine/Source/Programs/TraceAnalyzer/Private/TextSerializer.cpp
+++ b/Engine/Source/Programs/TraceAnalyzer/Private/TextSerializer.cpp
@@ -251,10 +251,10 @@ bool FStdoutTextSerializer::Commit()
 FFileTextSerializer::FFileTextSerializer(FileHandle InHandle)
 	: Handle(InHandle)
 {
-	Buffer = new uint8[BufferSize];
+	Buffer = FMemory::Malloc(BufferSize);
 	check(Buffer);
 
-	FormatBuffer = new uint8[FormatBufferSize];
+	FormatBuffer = FMemory::Malloc(FormatBufferSize);
 	check(FormatBuffer);
 }
 
@@ -262,8 +262,8 @@ FFileTextSerializer::FFileTextSerializer(FileHandle InHandle)
 
 FFileTextSerializer::~FFileTextSerializer()
 {
-	delete[] FormatBuffer;
-	delete[] Buffer;
+	FMemory::Free(FormatBuffer);
+	FMemory::Free(Buffer);
 }
 
 ////////////////////////////////////////////////////////////////////////////////////////////////////
diff --git a/Engine/Source/Programs/TraceAnalyzer/Private/TextSerializer.h b/Engine/Source/Programs/TraceAnalyzer/Private/TextSerializer.h
--- a/Engine/Source/Programs/TraceAnalyzer/Private/TextSerializer.h
+++ b/Engine/Source/Programs/TraceAnalyzer/Private/TextSerializer.h
@@ -4,8 +4,10 @@
 
 #include "CoreTypes.h"
 
-#include "Trace/Trace.h" // TraceLog, for UE::Trace::FEventRef8
+// TraceLog
+#include "Trace/Trace.h" // for UE::Trace::FEventRef8
 
+// TraceAnalyzer
 #include "Io.h"
 
 namespace UE
diff --git a/Engine/Source/Programs/TraceAnalyzer/Private/TraceAnalyzer.cpp b/Engine/Source/Programs/TraceAnalyzer/Private/TraceAnalyzer.cpp
--- a/Engine/Source/Programs/TraceAnalyzer/Private/TraceAnalyzer.cpp
+++ b/Engine/Source/Programs/TraceAnalyzer/Private/TraceAnalyzer.cpp
@@ -2,15 +2,15 @@
 
 #include "TraceAnalyzer.h"
 
-#include "CoreMinimal.h"
 #include "Misc/CommandLine.h"
 #include "RequiredProgramMainCPPInclude.h"
 
+// TraceAnalyzer
 #include "Command.h"
 
 #include <stdio.h>
 
-DEFINE_LOG_CATEGORY_STATIC(LogTraceAnalyzer, Log, All);
+DEFINE_LOG_CATEGORY(LogTraceAnalyzer);
 
 IMPLEMENT_APPLICATION(TraceAnalyzer, "TraceAnalyzer");
 
diff --git a/Engine/Source/Programs/TraceAnalyzer/Private/TraceAnalyzer.h b/Engine/Source/Programs/TraceAnalyzer/Private/TraceAnalyzer.h
--- a/Engine/Source/Programs/TraceAnalyzer/Private/TraceAnalyzer.h
+++ b/Engine/Source/Programs/TraceAnalyzer/Private/TraceAnalyzer.h
@@ -3,3 +3,22 @@
 #pragma once
 
 #include "CoreTypes.h"
+#include "Logging/LogCategory.h"
+
+DECLARE_LOG_CATEGORY_EXTERN(LogTraceAnalyzer, Log, All);
+
+#define LogTraceAnalyzerError(Format, ...) \
+	{ \
+		wprintf(L"Error: " Format "\n", ##__VA_ARGS__); \
+		UE_LOG(LogTraceAnalyzer, Error, TEXT(Format), ##__VA_ARGS__); \
+	}
+#define LogTraceAnalyzerWarning(Format, ...) \
+	{ \
+		wprintf(L"Warning: " Format "\n", ##__VA_ARGS__); \
+		UE_LOG(LogTraceAnalyzer, Warning, TEXT(Format), ##__VA_ARGS__); \
+	}
+#define LogTraceAnalyzerMessage(Format, ...) \
+	{ \
+		wprintf(L"" Format "\n", ##__VA_ARGS__); \
+		UE_LOG(LogTraceAnalyzer, Log, TEXT(Format), ##__VA_ARGS__); \
+	}
